Cloth Editor: enable key event handling in cloth paint tool

diff --git a/Engine/Plugins/Experimental/ChaosClothAssetEditor/Source/ChaosClothAssetEditor/Private/ChaosClothAsset/ClothEditorCommands.cpp b/Engine/Plugins/Experimental/ChaosClothAssetEditor/Source/ChaosClothAssetEditor/Private/ChaosClothAsset/ClothEditorCommands.cpp
--- a/Engine/Plugins/Experimental/ChaosClothAssetEditor/Source/ChaosClothAssetEditor/Private/ChaosClothAsset/ClothEditorCommands.cpp
+++ b/Engine/Plugins/Experimental/ChaosClothAssetEditor/Source/ChaosClothAssetEditor/Private/ChaosClothAsset/ClothEditorCommands.cpp
@@ -2,6 +2,7 @@
 
 #include "ChaosClothAsset/ClothEditorCommands.h"
 #include "ChaosClothAsset/ClothEditorStyle.h"
+#include "ChaosClothAsset/ClothWeightMapPaintTool.h"
 
 #define LOCTEXT_NAMESPACE "FChaosClothAssetEditorCommands"
 
@@ -38,5 +39,25 @@ void FChaosClothAssetEditorCommands::RegisterCommands()
 	UI_COMMAND(ToggleRenderMeshWireframe, "ToggleRenderMeshWireframe", "Toggle render mesh wireframe", EUserInterfaceActionType::Button, FInputChord());
 }
 
+void FChaosClothAssetEditorCommands::GetToolDefaultObjectList(TArray<UInteractiveTool*>& ToolCDOs)
+{
+	ToolCDOs.Add(GetMutableDefault<UClothEditorWeightMapPaintTool>());
+}
+
+void FChaosClothAssetEditorCommands::UpdateToolCommandBinding(UInteractiveTool* Tool, TSharedPtr<FUICommandList> UICommandList, bool bUnbind)
+{
+	if (FChaosClothAssetEditorCommands::IsRegistered())
+	{
+		if (bUnbind)
+		{
+			FChaosClothAssetEditorCommands::Get().UnbindActiveCommands(UICommandList);
+		}
+		else
+		{
+			FChaosClothAssetEditorCommands::Get().BindCommandsForCurrentTool(UICommandList, Tool);
+		}
+	}
+}
+
 
 #undef LOCTEXT_NAMESPACE
diff --git a/Engine/Plugins/Experimental/ChaosClothAssetEditor/Source/ChaosClothAssetEditor/Private/ChaosClothAsset/ClothEditorMode.cpp b/Engine/Plugins/Experimental/ChaosClothAssetEditor/Source/ChaosClothAssetEditor/Private/ChaosClothAsset/ClothEditorMode.cpp
--- a/Engine/Plugins/Experimental/ChaosClothAssetEditor/Source/ChaosClothAssetEditor/Private/ChaosClothAsset/ClothEditorMode.cpp
+++ b/Engine/Plugins/Experimental/ChaosClothAssetEditor/Source/ChaosClothAssetEditor/Private/ChaosClothAsset/ClothEditorMode.cpp
@@ -177,11 +177,15 @@ void UChaosClothAssetEditorMode::CreateToolkit()
 
 void UChaosClothAssetEditorMode::OnToolStarted(UInteractiveToolManager* Manager, UInteractiveTool* Tool)
 {
+	FChaosClothAssetEditorCommands::UpdateToolCommandBinding(Tool, ToolCommandList, false);
+
 	bCanTogglePattern2DMode = false;
 }
 
 void UChaosClothAssetEditorMode::OnToolEnded(UInteractiveToolManager* Manager, UInteractiveTool* Tool)
 {
+	FChaosClothAssetEditorCommands::UpdateToolCommandBinding(Tool, ToolCommandList, true);
+
 	bCanTogglePattern2DMode = true;
 
 	UpdateSimulationMeshes();
@@ -806,6 +810,7 @@ void UChaosClothAssetEditorMode::SetRestSpaceViewportClient(TWeakPtr<FChaosCloth
 	if (VC.IsValid())
 	{
 		VC->Set2DMode(bPattern2DMode);
+		VC->SetToolCommandList(ToolCommandList);
 	}
 }
 
diff --git a/Engine/Plugins/Experimental/ChaosClothAssetEditor/Source/ChaosClothAssetEditor/Private/ChaosClothAsset/ClothEditorRestSpaceViewportClient.cpp b/Engine/Plugins/Experimental/ChaosClothAssetEditor/Source/ChaosClothAssetEditor/Private/ChaosClothAsset/ClothEditorRestSpaceViewportClient.cpp
--- a/Engine/Plugins/Experimental/ChaosClothAssetEditor/Source/ChaosClothAssetEditor/Private/ChaosClothAsset/ClothEditorRestSpaceViewportClient.cpp
+++ b/Engine/Plugins/Experimental/ChaosClothAssetEditor/Source/ChaosClothAssetEditor/Private/ChaosClothAsset/ClothEditorRestSpaceViewportClient.cpp
@@ -95,6 +95,17 @@ bool FChaosClothEditorRestSpaceViewportClient::ShouldOrbitCamera() const
 
 bool FChaosClothEditorRestSpaceViewportClient::InputKey(const FInputKeyEventArgs& EventArgs)
 {
+	// See if any tool commands want to handle the key event
+	const TSharedPtr<FUICommandList> PinnedToolCommandList = ToolCommandList.Pin();
+	if (EventArgs.Event != IE_Released && PinnedToolCommandList.IsValid())
+	{
+		const FModifierKeysState KeyState = FSlateApplication::Get().GetModifierKeys();
+		if (PinnedToolCommandList->ProcessCommandBindings(EventArgs.Key, KeyState, (EventArgs.Event == IE_Repeat)))
+		{
+			return true;
+		}
+	}
+
 	if (!b2DMode)
 	{
 		return FEditorViewportClient::InputKey(EventArgs);
@@ -210,3 +221,8 @@ void FChaosClothEditorRestSpaceViewportClient::SetEditorViewportWidget(TWeakPtr<
 {
 	EditorViewportWidget = InEditorViewportWidget;
 }
+
+void FChaosClothEditorRestSpaceViewportClient::SetToolCommandList(TWeakPtr<FUICommandList> InToolCommandList)
+{
+	ToolCommandList = InToolCommandList;
+}
\ No newline at end of file
diff --git a/Engine/Plugins/Experimental/ChaosClothAssetEditor/Source/ChaosClothAssetEditor/Public/ChaosClothAsset/ClothEditorCommands.h b/Engine/Plugins/Experimental/ChaosClothAssetEditor/Source/ChaosClothAssetEditor/Public/ChaosClothAsset/ClothEditorCommands.h
--- a/Engine/Plugins/Experimental/ChaosClothAssetEditor/Source/ChaosClothAssetEditor/Public/ChaosClothAsset/ClothEditorCommands.h
+++ b/Engine/Plugins/Experimental/ChaosClothAssetEditor/Source/ChaosClothAssetEditor/Public/ChaosClothAsset/ClothEditorCommands.h
@@ -13,6 +13,17 @@ public:
 	// TBaseCharacterFXEditorCommands<> interface
 	virtual void RegisterCommands() override;
 
+	// TInteractiveToolCommands<>
+	virtual void GetToolDefaultObjectList(TArray<UInteractiveTool*>& ToolCDOs) override;
+
+	/**
+	 * Add or remove commands relevant to Tool to the given UICommandList.
+	 * Call this when the active tool changes (eg on ToolManager.OnToolStarted / OnToolEnded)
+	 * @param bUnbind if true, commands are removed, otherwise added
+	 */
+	static void UpdateToolCommandBinding(UInteractiveTool* Tool, TSharedPtr<FUICommandList> UICommandList, bool bUnbind = false);
+
+
 public:
 
 	TSharedPtr<FUICommandInfo> OpenClothEditor;
diff --git a/Engine/Plugins/Experimental/ChaosClothAssetEditor/Source/ChaosClothAssetEditor/Public/ChaosClothAsset/ClothEditorRestSpaceViewportClient.h b/Engine/Plugins/Experimental/ChaosClothAssetEditor/Source/ChaosClothAssetEditor/Public/ChaosClothAsset/ClothEditorRestSpaceViewportClient.h
--- a/Engine/Plugins/Experimental/ChaosClothAssetEditor/Source/ChaosClothAssetEditor/Public/ChaosClothAsset/ClothEditorRestSpaceViewportClient.h
+++ b/Engine/Plugins/Experimental/ChaosClothAssetEditor/Source/ChaosClothAssetEditor/Public/ChaosClothAsset/ClothEditorRestSpaceViewportClient.h
@@ -31,6 +31,7 @@ public:
 	void Set2DMode(bool In2DMode);
 
 	void SetEditorViewportWidget(TWeakPtr<SEditorViewport> InEditorViewportWidget);
+	void SetToolCommandList(TWeakPtr<FUICommandList> ToolCommandList);
 
 private:
 
@@ -42,5 +43,7 @@ private:
 
 	TUniquePtr<FEditor2DScrollBehaviorTarget> ScrollBehaviorTarget;
 	TUniquePtr<FEditor2DMouseWheelZoomBehaviorTarget> ZoomBehaviorTarget;
+
+	TWeakPtr<FUICommandList> ToolCommandList;
 };
 
diff --git a/Engine/Plugins/Experimental/CharacterFXEditor/BaseCharacterFXEditor/Source/BaseCharacterFXEditor/Public/BaseCharacterFXEditorCommands.h b/Engine/Plugins/Experimental/CharacterFXEditor/BaseCharacterFXEditor/Source/BaseCharacterFXEditor/Public/BaseCharacterFXEditorCommands.h
--- a/Engine/Plugins/Experimental/CharacterFXEditor/BaseCharacterFXEditor/Source/BaseCharacterFXEditor/Public/BaseCharacterFXEditorCommands.h
+++ b/Engine/Plugins/Experimental/CharacterFXEditor/BaseCharacterFXEditor/Source/BaseCharacterFXEditor/Public/BaseCharacterFXEditorCommands.h
@@ -2,22 +2,24 @@
 
 #pragma once
 
-#include "Framework/Commands/Commands.h"
+#include "Tools/InteractiveToolsCommands.h"
 
 /**
 * Home for FUICommandInfos used by the editor
 */
 
 template<typename CommandContextType>
-class TBaseCharacterFXEditorCommands : public TCommands<CommandContextType>
+class TBaseCharacterFXEditorCommands : public TInteractiveToolCommands<CommandContextType>
 {
 
 public:
 
-	// TCommands<> interface
+	// TInteractiveToolCommands<> interface
 	// If you override this function, call this base version to set up Accept/Cancel buttons
 	virtual void RegisterCommands() override
 	{
+		TInteractiveToolCommands<CommandContextType>::RegisterCommands();
+
 #define LOCTEXT_NAMESPACE "TBaseCharacterFXEditorCommands"
 		// These allow us to link up to pressed keys
 		UI_COMMAND(AcceptOrCompleteActiveTool, "Accept", "Accept the active tool", EUserInterfaceActionType::Button, FInputChord(EKeys::Enter));
@@ -32,7 +34,7 @@ protected:
 
 	// Call this from subclass default constructor
 	TBaseCharacterFXEditorCommands(const FName InContextName, const FText& InContextDesc, const FName InContextParent, const FName InStyleSetName) :
-		TCommands<CommandContextType>(InContextName, InContextDesc, InContextParent, InStyleSetName)
+		TInteractiveToolCommands<CommandContextType>(InContextName, InContextDesc, InContextParent, InStyleSetName)
 	{}
 
 };
diff --git a/Engine/Plugins/Experimental/CharacterFXEditor/ExampleCharacterFXEditor/Source/ExampleCharacterFXEditor/Public/ExampleCharacterFXEditorCommands.h b/Engine/Plugins/Experimental/CharacterFXEditor/ExampleCharacterFXEditor/Source/ExampleCharacterFXEditor/Public/ExampleCharacterFXEditorCommands.h
--- a/Engine/Plugins/Experimental/CharacterFXEditor/ExampleCharacterFXEditor/Source/ExampleCharacterFXEditor/Public/ExampleCharacterFXEditorCommands.h
+++ b/Engine/Plugins/Experimental/CharacterFXEditor/ExampleCharacterFXEditor/Source/ExampleCharacterFXEditor/Public/ExampleCharacterFXEditorCommands.h
@@ -10,8 +10,9 @@ public:
 
 	FExampleCharacterFXEditorCommands();
 
-	// TCommands<> interface
+	// TInteractiveToolCommands<> interface
 	virtual void RegisterCommands() override;
+	virtual void GetToolDefaultObjectList(TArray<UInteractiveTool*>& ToolCDOs) override {}
 
 public:
 
